New TRS Gizmos: fixed settings being saved in the DefaultEditorPerProjectUserSettings

diff --git a/Engine/Plugins/Experimental/GizmoFramework/Source/GizmoSettings/Public/GizmoSettings.h b/Engine/Plugins/Experimental/GizmoFramework/Source/GizmoSettings/Public/GizmoSettings.h
--- a/Engine/Plugins/Experimental/GizmoFramework/Source/GizmoSettings/Public/GizmoSettings.h
+++ b/Engine/Plugins/Experimental/GizmoFramework/Source/GizmoSettings/Public/GizmoSettings.h
@@ -8,7 +8,7 @@
 
 #include "GizmoSettings.generated.h"
 
-UCLASS(config = EditorPerProjectUserSettings, defaultconfig, meta = (DisplayName = "New TRS Gizmo"))
+UCLASS(config=EditorPerProjectUserSettings, meta = (DisplayName = "New TRS Gizmo"))
 class GIZMOSETTINGS_API UGizmoSettings : public UDeveloperSettings
 {
 	GENERATED_UCLASS_BODY()
