[Backout] - CL41482814
#rnx
Original CL Desc
-----------------------------------------------------------------
[Gizmos] New Transform Gizmo - Preference

- Creates new `UTransformGizmoEditorSettings` class stored in `EditorInteractiveFramework` module

- Removes local variables tracking new TRS on/off and Gizmo Settings from `GizmoManagerLocals` namespace in `EditorInteractiveGizmoManager.cpp`

- Switches to use the new variables in the Settings Class

- Deprecates `ULevelEditorViewportSettings::TransformWidgetSizeAdjustment` in favor of `UTransformGizmoEditorSettings::TransformGizmoSize`

- Maps new, non-config properties to the following found in `ULevelEditorViewportSettings`: `bAllowArcballRotate`, `bAllowScreenRotate`, `bAllowEditWidgetAxisDisplay`, `bAllowTranslateRotateZWidget`. This allows to avoid a circular dependency `EditorITF Module` <-> `LevelEditor` Module, and saves us from having config properties which will be eliminated in the future, since they are part of legacy gizmo settings

- Removes Experimental `GizmoFramework` plugin, which currently only contains settings which would be a duplicate of some of the ones introduced in this CL by `UTransformGizmoEditorSettings`

diff --git a/Engine/Plugins/Experimental/GizmoFramework/GizmoFramework.uplugin b/Engine/Plugins/Experimental/GizmoFramework/GizmoFramework.uplugin
new file mode 100644
--- /dev/null
+++ b/Engine/Plugins/Experimental/GizmoFramework/GizmoFramework.uplugin
@@ -0,0 +1,24 @@
+{
+	"FileVersion": 3,
+	"Version": 1,
+	"VersionName": "1.0",
+	"FriendlyName": "New TRS Gizmos",
+	"Description": "Toggle use of the new TRS Gizmos",
+	"Category": "Experimental",
+	"CreatedBy": "Epic Games, Inc.",
+	"CreatedByURL": "https://epicgames.com",
+	"DocsURL": "",
+	"MarketplaceURL": "",
+	"SupportURL": "",
+	"CanContainContent": false,
+	"IsBetaVersion": false,
+	"IsExperimentalVersion": true,
+	"Installed": false,
+	"Modules": [
+		{
+			"Name": "GizmoSettings",
+			"Type": "Editor",
+			"LoadingPhase": "Default"
+		}
+	]
+}
\ No newline at end of file
diff --git a/Engine/Plugins/Experimental/GizmoFramework/Source/GizmoSettings/GizmoSettings.Build.cs b/Engine/Plugins/Experimental/GizmoFramework/Source/GizmoSettings/GizmoSettings.Build.cs
new file mode 100644
--- /dev/null
+++ b/Engine/Plugins/Experimental/GizmoFramework/Source/GizmoSettings/GizmoSettings.Build.cs
@@ -0,0 +1,24 @@
+// Copyright Epic Games, Inc. All Rights Reserved.
+
+namespace UnrealBuildTool.Rules
+{
+	public class GizmoSettings : ModuleRules
+	{
+		public GizmoSettings(ReadOnlyTargetRules Target) : base(Target)
+		{
+			// Enable truncation warnings in this module
+			CppCompileWarningSettings.UnsafeTypeCastWarningLevel = WarningLevel.Warning;
+
+			PrivateDependencyModuleNames.AddRange(
+				new string[]
+				{
+					"Core",
+					"CoreUObject",
+					"DeveloperSettings",
+					"Engine",
+					"EditorInteractiveToolsFramework"
+				}
+			);
+		}
+	}
+}
diff --git a/Engine/Plugins/Experimental/GizmoFramework/Source/GizmoSettings/Private/GizmoSettings.cpp b/Engine/Plugins/Experimental/GizmoFramework/Source/GizmoSettings/Private/GizmoSettings.cpp
new file mode 100644
--- /dev/null
+++ b/Engine/Plugins/Experimental/GizmoFramework/Source/GizmoSettings/Private/GizmoSettings.cpp
@@ -0,0 +1,15 @@
+// Copyright Epic Games, Inc. All Rights Reserved.
+
+#include "GizmoSettings.h"
+
+UGizmoSettings::UGizmoSettings(const FObjectInitializer& ObjectInitlaizer)
+	: UDeveloperSettings(ObjectInitlaizer)
+{}
+
+UGizmoSettings::FOnUpdateSettings UGizmoSettings::OnSettingsChange;
+
+void UGizmoSettings::PostEditChangeProperty(FPropertyChangedEvent& PropertyChangedEvent)
+{
+	Super::PostEditChangeProperty(PropertyChangedEvent);
+	OnSettingsChange.Broadcast(this);
+}
diff --git a/Engine/Plugins/Experimental/GizmoFramework/Source/GizmoSettings/Private/GizmoSettingsModule.cpp b/Engine/Plugins/Experimental/GizmoFramework/Source/GizmoSettings/Private/GizmoSettingsModule.cpp
new file mode 100644
--- /dev/null
+++ b/Engine/Plugins/Experimental/GizmoFramework/Source/GizmoSettings/Private/GizmoSettingsModule.cpp
@@ -0,0 +1,30 @@
+// Copyright Epic Games, Inc. All Rights Reserved.
+
+#include "GizmoSettingsModule.h"
+#include "Modules/ModuleManager.h"
+
+#include "GizmoSettings.h"
+#include "EditorInteractiveGizmoManager.h"
+
+void FGizmoSettingsModule::StartupModule()
+{
+	UGizmoSettings::OnSettingsChange.AddRaw(this, &FGizmoSettingsModule::UpdateGizmos);
+	UpdateGizmos(GetDefault<UGizmoSettings>());
+}
+
+void FGizmoSettingsModule::ShutdownModule()
+{
+	UGizmoSettings::OnSettingsChange.RemoveAll(this);
+}
+
+void FGizmoSettingsModule::UpdateGizmos(const UGizmoSettings* InSettings)
+{
+	if (InSettings)
+	{
+		UEditorInteractiveGizmoManager::SetUsesNewTRSGizmos(InSettings->bEnableNewGizmos);
+		UEditorInteractiveGizmoManager::SetGizmosParameters(InSettings->GizmoParameters);
+	}
+}
+
+
+IMPLEMENT_MODULE(FGizmoSettingsModule, GizmoSettings)
diff --git a/Engine/Plugins/Experimental/GizmoFramework/Source/GizmoSettings/Public/GizmoSettings.h b/Engine/Plugins/Experimental/GizmoFramework/Source/GizmoSettings/Public/GizmoSettings.h
new file mode 100644
--- /dev/null
+++ b/Engine/Plugins/Experimental/GizmoFramework/Source/GizmoSettings/Public/GizmoSettings.h
@@ -0,0 +1,30 @@
+// Copyright Epic Games, Inc. All Rights Reserved.
+
+#pragma once
+
+#include "CoreMinimal.h"
+#include "Engine/DeveloperSettings.h"
+#include "EditorGizmos/TransformGizmo.h"
+
+#include "GizmoSettings.generated.h"
+
+UCLASS(config=EditorPerProjectUserSettings, meta = (DisplayName = "New TRS Gizmo"))
+class GIZMOSETTINGS_API UGizmoSettings : public UDeveloperSettings
+{
+	GENERATED_UCLASS_BODY()
+
+public:
+
+	/** Enable/Disable New TRS Gizmos across the editor. */
+	UPROPERTY(config, EditAnywhere, Category = "New TRS Gizmo")
+	bool bEnableNewGizmos = false;
+
+	/** Change the current gizmos parameters. */
+	UPROPERTY(config, EditAnywhere, Category = "New TRS Gizmo")
+	FGizmosParameters GizmoParameters;
+	
+	DECLARE_MULTICAST_DELEGATE_OneParam(FOnUpdateSettings,  const UGizmoSettings*);
+	static FOnUpdateSettings OnSettingsChange;
+
+	virtual void PostEditChangeProperty(FPropertyChangedEvent& PropertyChangedEvent) override;
+};
\ No newline at end of file
diff --git a/Engine/Plugins/Experimental/GizmoFramework/Source/GizmoSettings/Public/GizmoSettingsModule.h b/Engine/Plugins/Experimental/GizmoFramework/Source/GizmoSettings/Public/GizmoSettingsModule.h
new file mode 100644
--- /dev/null
+++ b/Engine/Plugins/Experimental/GizmoFramework/Source/GizmoSettings/Public/GizmoSettingsModule.h
@@ -0,0 +1,16 @@
+// Copyright Epic Games, Inc. All Rights Reserved.
+
+#pragma once
+
+#include "Modules/ModuleInterface.h"
+
+class UGizmoSettings;
+
+class FGizmoSettingsModule : public IModuleInterface
+{
+	/** IModuleInterface implementation */
+	virtual void StartupModule() override;
+	virtual void ShutdownModule() override;
+
+	void UpdateGizmos(const UGizmoSettings* InSettings);
+};
