[New TRS Gizmos] Restore UGizmoSettings for proper deprecation

diff --git a/Engine/Plugins/Experimental/GizmoFramework/Source/GizmoSettings/GizmoSettings.Build.cs b/Engine/Plugins/Experimental/GizmoFramework/Source/GizmoSettings/GizmoSettings.Build.cs
--- a/Engine/Plugins/Experimental/GizmoFramework/Source/GizmoSettings/GizmoSettings.Build.cs
+++ b/Engine/Plugins/Experimental/GizmoFramework/Source/GizmoSettings/GizmoSettings.Build.cs
@@ -12,7 +12,11 @@ namespace UnrealBuildTool.Rules
 			PrivateDependencyModuleNames.AddRange(
 				new string[]
 				{
-					"Core"
+					"Core",
+					"CoreUObject",
+                    "DeveloperSettings",
+                    "Engine",
+                    "EditorInteractiveToolsFramework"
 				}
 			);
 		}
diff --git a/Engine/Plugins/Experimental/GizmoFramework/Source/GizmoSettings/Private/GizmoSettings.cpp b/Engine/Plugins/Experimental/GizmoFramework/Source/GizmoSettings/Private/GizmoSettings.cpp
new file mode 100644
--- /dev/null
+++ b/Engine/Plugins/Experimental/GizmoFramework/Source/GizmoSettings/Private/GizmoSettings.cpp
@@ -0,0 +1,36 @@
+// Copyright Epic Games, Inc. All Rights Reserved.
+
+#include "GizmoSettings.h"
+#include "EditorInteractiveGizmoManager.h"
+
+void UGizmoSettings::PostInitProperties()
+{
+	Super::PostInitProperties();
+
+	bool bConfigNeedsUpdate = false;
+
+	PRAGMA_DISABLE_DEPRECATION_WARNINGS
+	if (bEnableNewGizmos_DEPRECATED)
+	{
+		UEditorInteractiveGizmoManager::SetUsesNewTRSGizmos(bEnableNewGizmos_DEPRECATED);
+		bEnableNewGizmos_DEPRECATED = false;
+
+		bConfigNeedsUpdate = true;
+	}
+
+	FGizmosParameters DefaultGizmosParameters;
+	if (GizmoParameters_DEPRECATED.bCtrlMiddleDoesY != DefaultGizmosParameters.bCtrlMiddleDoesY ||
+		GizmoParameters_DEPRECATED.bEnableExplicit != DefaultGizmosParameters.bEnableExplicit ||
+		GizmoParameters_DEPRECATED.RotateMode != DefaultGizmosParameters.RotateMode)
+	{
+		UEditorInteractiveGizmoManager::SetGizmosParameters(GizmoParameters_DEPRECATED);
+
+		bConfigNeedsUpdate = true;
+	}
+	PRAGMA_ENABLE_DEPRECATION_WARNINGS
+
+	if (bConfigNeedsUpdate)
+	{
+		SaveConfig();
+	}
+}
diff --git a/Engine/Plugins/Experimental/GizmoFramework/Source/GizmoSettings/Public/GizmoSettings.h b/Engine/Plugins/Experimental/GizmoFramework/Source/GizmoSettings/Public/GizmoSettings.h
new file mode 100644
--- /dev/null
+++ b/Engine/Plugins/Experimental/GizmoFramework/Source/GizmoSettings/Public/GizmoSettings.h
@@ -0,0 +1,27 @@
+// Copyright Epic Games, Inc. All Rights Reserved.
+
+#pragma once
+
+#include "Engine/DeveloperSettings.h"
+#include "EditorGizmos/TransformGizmo.h"
+
+#include "GizmoSettings.generated.h"
+
+UCLASS(config=EditorPerProjectUserSettings, meta = (DisplayName = "New TRS Gizmo"))
+class UE_DEPRECATED(5.6, "New Gizmos settings can now be found in UTransformGizmoEditorSettings") GIZMOSETTINGS_API UGizmoSettings : public UDeveloperSettings
+{
+	GENERATED_BODY()
+
+public:
+	/** Enable/Disable New TRS Gizmos across the editor. */
+	UE_DEPRECATED(5.6, "UGizmoSettings::bEnableNewGizmos is deprecated. See UTransformGizmoEditorSettings::bUseExperimentalGizmo")
+	UPROPERTY(config)
+	bool bEnableNewGizmos_DEPRECATED = false;
+
+	/** Change the current gizmos parameters. */
+	UE_DEPRECATED(5.6, "UGizmoSettings::GizmoParameters is deprecated. See UTransformGizmoEditorSettings::GizmoParameters")
+	UPROPERTY(config)
+	FGizmosParameters GizmoParameters_DEPRECATED;
+
+	virtual void PostInitProperties() override;
+};
