[Hide deprecated "New TRS Gizmo" settings from Editor Preferences]

diff --git a/Engine/Plugins/Experimental/GizmoFramework/Source/GizmoSettings/Public/GizmoSettings.h b/Engine/Plugins/Experimental/GizmoFramework/Source/GizmoSettings/Public/GizmoSettings.h
--- a/Engine/Plugins/Experimental/GizmoFramework/Source/GizmoSettings/Public/GizmoSettings.h
+++ b/Engine/Plugins/Experimental/GizmoFramework/Source/GizmoSettings/Public/GizmoSettings.h
@@ -24,4 +24,7 @@ public:
 	FGizmosParameters GizmoParameters_DEPRECATED;
 
 	virtual void PostInitProperties() override;
+
+	/** We don't want to show this in the Editor Preferences list */
+	virtual bool SupportsAutoRegistration() const override { return false; }
 };
