[ITF Gizmos] Rename the animation gizmo and TRS gizmo plugins to reduce confusion. Rename the temporary TRS gizmo preference options.

#rnx

diff --git a/Engine/Plugins/Experimental/GizmoFramework/GizmoFramework.uplugin b/Engine/Plugins/Experimental/GizmoFramework/GizmoFramework.uplugin
--- a/Engine/Plugins/Experimental/GizmoFramework/GizmoFramework.uplugin
+++ b/Engine/Plugins/Experimental/GizmoFramework/GizmoFramework.uplugin
@@ -2,7 +2,7 @@
 	"FileVersion": 3,
 	"Version": 1,
 	"VersionName": "1.0",
-	"FriendlyName": "New TRS Gizmos",
+	"FriendlyName": "Animator Gizmos",
 	"Description": "Toggle use of the new TRS Gizmos",
 	"Category": "Experimental",
 	"CreatedBy": "Epic Games, Inc.",
